- Enabled CookedCooker in the engine
- Made an EngineTestCookedCooker target

diff --git a/Engine/Config/BaseGame.ini b/Engine/Config/BaseGame.ini
--- a/Engine/Config/BaseGame.ini
+++ b/Engine/Config/BaseGame.ini
@@ -265,7 +265,6 @@ ReleaseTargetType=Game
 
 
 [CookedEditorSettings_CookedEditor]
-bIsForExternalDistribution=true
 
 ; currently don't work well with cooked content
 +DisabledPlugins=GeometryMode
diff --git a/Engine/Source/Developer/CookedEditor/Public/CookedEditorTargetPlatform.h b/Engine/Source/Developer/CookedEditor/Public/CookedEditorTargetPlatform.h
--- a/Engine/Source/Developer/CookedEditor/Public/CookedEditorTargetPlatform.h
+++ b/Engine/Source/Developer/CookedEditor/Public/CookedEditorTargetPlatform.h
@@ -125,14 +125,13 @@ public:
 		PackageManager = ICookedEditorPackageManager::FactoryForTargetPlatform(this, true);
 	}
 
-	TCookedCookerTargetPlatform(ICookedEditorPackageManager* ExistingManager)
+	TCookedCookerTargetPlatform(TUniquePtr<ICookedEditorPackageManager>&& ExistingManager)
 	{
 		PackageManager = MoveTemp(ExistingManager);
 	}
 
 	~TCookedCookerTargetPlatform()
 	{
-		delete PackageManager;
 	}
 
 
@@ -227,6 +226,6 @@ public:
 		return false;
 	}
 
-	ICookedEditorPackageManager* PackageManager;
+	TUniquePtr<ICookedEditorPackageManager> PackageManager;
 
 };
diff --git a/Engine/Source/Developer/Windows/WindowsTargetPlatform/Private/WindowsTargetPlatformModule.cpp b/Engine/Source/Developer/Windows/WindowsTargetPlatform/Private/WindowsTargetPlatformModule.cpp
--- a/Engine/Source/Developer/Windows/WindowsTargetPlatform/Private/WindowsTargetPlatformModule.cpp
+++ b/Engine/Source/Developer/Windows/WindowsTargetPlatform/Private/WindowsTargetPlatformModule.cpp
@@ -61,6 +61,7 @@ public:
 		if (GConfig->GetBool(TEXT("CookedEditorSettings"), TEXT("bSupportCookedEditor"), bSupportCookedEditor, GGameIni) && bSupportCookedEditor)
 		{
 			TargetPlatforms.Add(new TCookedEditorTargetPlatform<FWindowsEditorTargetPlatformParent>());
+			TargetPlatforms.Add(new TCookedCookerTargetPlatform<FWindowsEditorTargetPlatformParent>());
 		}
 #endif
 	}
