- Fixed some cooker commandline options to match Old School cooked cooker method
- Removed a random None package name to fix a warning

diff --git a/Engine/Source/Developer/CookedEditor/Private/CookedEditorPackageManager.cpp b/Engine/Source/Developer/CookedEditor/Private/CookedEditorPackageManager.cpp
--- a/Engine/Source/Developer/CookedEditor/Private/CookedEditorPackageManager.cpp
+++ b/Engine/Source/Developer/CookedEditor/Private/CookedEditorPackageManager.cpp
@@ -148,6 +148,7 @@ void FIniCookedEditorPackageManager::FilterGatheredPackages(TArray<FName>& Packa
 			}
 			return false;
 		});
+	PackageNames.Remove(NAME_None);
 }
 
 void FIniCookedEditorPackageManager::GetEnginePackagesToCook(TArray<FName>& PackagesToCook) const
