Optimize FIniCookedEditorPackageManager::FilterGatheredPackages by avoiding allocations

diff --git a/Engine/Source/Developer/CookedEditor/Private/CookedEditorPackageManager.cpp b/Engine/Source/Developer/CookedEditor/Private/CookedEditorPackageManager.cpp
--- a/Engine/Source/Developer/CookedEditor/Private/CookedEditorPackageManager.cpp
+++ b/Engine/Source/Developer/CookedEditor/Private/CookedEditorPackageManager.cpp
@@ -141,9 +141,12 @@ void FIniCookedEditorPackageManager::FilterGatheredPackages(TArray<FName>& Packa
 	// now filter based on ini settings
 	PackageNames.RemoveAll([this](FName& AssetPath)
 		{
+			FNameBuilder AssetPathBuilder;
+			AssetPath.ToString(AssetPathBuilder);
+			const FStringView AssetPathView(AssetPathBuilder);
 			for (const FString& Path : DisallowedPathsToGather)
 			{
-				if (AssetPath.ToString().StartsWith(Path))
+				if (AssetPathView.StartsWith(Path))
 				{
 					return true;
 				}
