Cleanup
#rnx

diff --git a/Engine/Source/Developer/CookedEditor/Private/CookedEditorPackageManager.cpp b/Engine/Source/Developer/CookedEditor/Private/CookedEditorPackageManager.cpp
--- a/Engine/Source/Developer/CookedEditor/Private/CookedEditorPackageManager.cpp
+++ b/Engine/Source/Developer/CookedEditor/Private/CookedEditorPackageManager.cpp
@@ -141,8 +141,7 @@ void FIniCookedEditorPackageManager::FilterGatheredPackages(TArray<FName>& Packa
 	// now filter based on ini settings
 	PackageNames.RemoveAll([this](FName& AssetPath)
 		{
-			FNameBuilder AssetPathBuilder;
-			AssetPath.ToString(AssetPathBuilder);
+			FNameBuilder AssetPathBuilder(AssetPath);
 			const FStringView AssetPathView(AssetPathBuilder);
 			for (const FString& Path : DisallowedPathsToGather)
 			{
