Fixing bad module reference.

diff --git a/Engine/Source/Developer/CookedEditor/CookedEditor.Build.cs b/Engine/Source/Developer/CookedEditor/CookedEditor.Build.cs
--- a/Engine/Source/Developer/CookedEditor/CookedEditor.Build.cs
+++ b/Engine/Source/Developer/CookedEditor/CookedEditor.Build.cs
@@ -1,7 +1,7 @@
 // Copyright Epic Games, Inc. All Rights Reserved.
 
-using UnrealBuildTool;
 using System.IO;
+using UnrealBuildTool;
 
 public class CookedEditor : ModuleRules
 {
@@ -28,13 +28,13 @@ public class CookedEditor : ModuleRules
 				"TargetPlatform",
 			});
 
+		PrivateIncludePathModuleNames.Add("WindowsTargetPlatform");
 
 		// we currently need to pull in desktop templates that are private
 		// @todo move these out to Public so they can be extended like this
 		string EnginePath = Path.GetFullPath(Target.RelativeEnginePath);
 		PublicIncludePaths.AddRange(
 			new string[] {
-				Path.Combine(EnginePath, "Source/Developer/Windows/WindowsTargetPlatform/Private"),
 				Path.Combine(EnginePath, "Source/Developer/Linux/LinuxTargetPlatform/Private"),
 				Path.Combine(EnginePath, "Source/Developer/Mac/MacTargetPlatform/Private"),
 			}
