Attempting to fix public include paths.

diff --git a/Engine/Source/Developer/CookedEditor/CookedEditor.Build.cs b/Engine/Source/Developer/CookedEditor/CookedEditor.Build.cs
--- a/Engine/Source/Developer/CookedEditor/CookedEditor.Build.cs
+++ b/Engine/Source/Developer/CookedEditor/CookedEditor.Build.cs
@@ -28,7 +28,7 @@ public class CookedEditor : ModuleRules
 				"TargetPlatform",
 			});
 
-		PrivateIncludePathModuleNames.Add("WindowsTargetPlatform");
+		PublicIncludePathModuleNames.Add("WindowsTargetPlatform");
 
 		// we currently need to pull in desktop templates that are private
 		// @todo move these out to Public so they can be extended like this
