Fixing static analysis issue caused by mixing includes and code.

diff --git a/Engine/Source/Developer/CookedEditor/Public/CookedEditorPackageManager.h b/Engine/Source/Developer/CookedEditor/Public/CookedEditorPackageManager.h
--- a/Engine/Source/Developer/CookedEditor/Public/CookedEditorPackageManager.h
+++ b/Engine/Source/Developer/CookedEditor/Public/CookedEditorPackageManager.h
@@ -9,15 +9,21 @@
 #include "GenericWindowsTargetPlatform.h"
 #include "Windows/WindowsPlatformProperties.h"
 
+#if COOKEDEDITOR_WITH_LINUXTARGETPLATFORM
+	#include "LinuxTargetPlatform.h"
+#endif
+
+#if COOKEDEDITOR_WITH_MACTARGETPLATFORM
+	#include "GenericMacTargetPlatform.h"
+#endif
+
 typedef TGenericWindowsTargetPlatform<FWindowsPlatformProperties<false, false, false>> FWindowsEditorTargetPlatformParent;
 
 #if COOKEDEDITOR_WITH_LINUXTARGETPLATFORM
-#include "LinuxTargetPlatform.h"
 typedef TLinuxTargetPlatform<FLinuxPlatformProperties<false, false, false, false>> FLinuxEditorTargetPlatformParent;
 #endif
 
 #if COOKEDEDITOR_WITH_MACTARGETPLATFORM
-#include "GenericMacTargetPlatform.h"
 typedef TGenericMacTargetPlatform<false, false, false> FMacEditorTargetPlatformParent;
 #endif
 
