Expose GetConfigArray
#rnx

diff --git a/Engine/Source/Developer/CookedEditor/Private/CookedEditorPackageManager.cpp b/Engine/Source/Developer/CookedEditor/Private/CookedEditorPackageManager.cpp
--- a/Engine/Source/Developer/CookedEditor/Private/CookedEditorPackageManager.cpp
+++ b/Engine/Source/Developer/CookedEditor/Private/CookedEditorPackageManager.cpp
@@ -118,6 +118,11 @@ FIniCookedEditorPackageManager::FIniCookedEditorPackageManager(bool bInIsCookedC
 }
 
 TArray<FString> FIniCookedEditorPackageManager::GetConfigArray(const TCHAR* Key) const
+{
+	return GetConfigArray(Key, bIsCookedCooker);
+}
+
+TArray<FString> FIniCookedEditorPackageManager::GetConfigArray(const TCHAR* Key, bool bIsCookedCooker)
 {
 	const TCHAR* SharedIniSection = TEXT("CookedEditorSettings");
 	const TCHAR* SpecificiIniSection = bIsCookedCooker ? TEXT("CookedEditorSettings_CookedCooker") : TEXT("CookedEditorSettings_CookedEditor");
diff --git a/Engine/Source/Developer/CookedEditor/Public/CookedEditorPackageManager.h b/Engine/Source/Developer/CookedEditor/Public/CookedEditorPackageManager.h
--- a/Engine/Source/Developer/CookedEditor/Public/CookedEditorPackageManager.h
+++ b/Engine/Source/Developer/CookedEditor/Public/CookedEditorPackageManager.h
@@ -167,6 +167,9 @@ public:
 
 	FIniCookedEditorPackageManager(bool bIsCookedCooker);
 
+	// gets an array from two sections, depending on bIsCookedCooker setting
+	static TArray<FString> GetConfigArray(const TCHAR* Key, bool bIsCookedCooker);
+
 	virtual void FilterGatheredPackages(TArray<FName>& PackageNames) const override;
 	virtual void GetEnginePackagesToCook(TArray<FName>& PackagesToCook) const override;
 	virtual void GetProjectPackagesToCook(TArray<FName>& PackagesToCook) const override;
