CookedEditor: Support native classes, and wildcards to specify them, in [CookSettings].DisallowedObjectClassNamesToLoad.
#rnx

diff --git a/Engine/Source/Developer/CookedEditor/Private/CookedEditorPackageManager.cpp b/Engine/Source/Developer/CookedEditor/Private/CookedEditorPackageManager.cpp
--- a/Engine/Source/Developer/CookedEditor/Private/CookedEditorPackageManager.cpp
+++ b/Engine/Source/Developer/CookedEditor/Private/CookedEditorPackageManager.cpp
@@ -1,6 +1,8 @@
 // Copyright Epic Games, Inc. All Rights Reserved.
 
 #include "CookedEditorPackageManager.h"
+#include "Algo/Sort.h"
+#include "Algo/Unique.h"
 #include "AssetRegistry/AssetData.h"
 #include "CoreMinimal.h"
 #include "Logging/LogMacros.h"
@@ -95,17 +97,79 @@ FIniCookedEditorPackageManager::FIniCookedEditorPackageManager(bool bInIsCookedC
 	ProjectAssetPaths = GetConfigArray(TEXT("ProjectAssetPaths"));
 	DisallowedPathsToGather = GetConfigArray(TEXT("DisallowedPathsToGather"));
 	DisabledPlugins = GetConfigArray(TEXT("DisabledPlugins"));
+}
+
+void FIniCookedEditorPackageManager::InitializeClasses()
+{
+	if (bClassesInitialized)
+	{
+		return;
+	}
+	bClassesInitialized = true;
+
+	IAssetRegistry& AssetRegistry = IAssetRegistry::GetChecked();
+	// SearchAllAssets should be a noop because the cooker already did it, but run it just in case.
+	AssetRegistry.SearchAllAssets(true /* bSynchronousSearch */);
 
+	constexpr TCHAR WildcardCharacter = '*';
+
+	TArray<FString> DisallowedObjectClassWildcards;
 	TArray<FString> DisallowedObjectClassNamesToLoad = GetConfigArray(TEXT("DisallowedObjectClassesToLoad"));;
+	for (TArray<FString>::TIterator Iter(DisallowedObjectClassNamesToLoad); Iter; ++Iter)
+	{
+		FString& ClassName = *Iter;
+		int32 UnusedIndex;
+		if (ClassName.FindChar(WildcardCharacter, UnusedIndex))
+		{
+			DisallowedObjectClassWildcards.Add(MoveTemp(ClassName));
+			Iter.RemoveCurrentSwap();
+		}
+	}
+	if (!DisallowedObjectClassWildcards.IsEmpty())
+	{
+		// Get the list of all blueprint and native classes from the AssetRegistry
+		TArray<FTopLevelAssetPath> ClassNamesUObject;
+		ClassNamesUObject.Add(FTopLevelAssetPath(UObject::StaticClass()));
+		TSet<FTopLevelAssetPath> AllClasses;
+		AssetRegistry.GetDerivedClassNames(ClassNamesUObject, TSet<FTopLevelAssetPath>(), AllClasses);
+
+		// Check the wildcards from config against each classpath
+		TSet<FString> MatchingClassPaths;
+		FString ClassPathStr;
+		for (FTopLevelAssetPath& ClassPath : AllClasses)
+		{
+			ClassPath.ToString(ClassPathStr);
+			bool bMatches = false;
+			for (const FString& Wildcard : DisallowedObjectClassWildcards)
+			{
+				if (ClassPathStr.MatchesWildcard(Wildcard, ESearchCase::IgnoreCase))
+				{
+					bMatches = true;
+					break;
+				}
+			}
+			if (bMatches)
+			{
+				DisallowedObjectClassNamesToLoad.Add(ClassPathStr);
+			}
+		}
+	}
+	Algo::Sort(DisallowedObjectClassNamesToLoad);
+	DisallowedObjectClassNamesToLoad.SetNum(Algo::Unique(DisallowedObjectClassNamesToLoad));
+	TArray<FTopLevelAssetPath> DisallowedObjectBaseClassPaths;
+	DisallowedObjectBaseClassPaths.Reserve(DisallowedObjectClassNamesToLoad.Num());
 	for (const FString& ClassName : DisallowedObjectClassNamesToLoad)
 	{
 		check(FPackageName::IsValidObjectPath(ClassName));
-		UClass* Class = FindObject<UClass>(nullptr, *ClassName);
-		check(Class);
-
-		DisallowedObjectClassesToLoad.Add(Class);
+		FTopLevelAssetPath ClassPath(ClassName);
+		check(ClassPath.IsValid());
+		DisallowedObjectBaseClassPaths.Add(ClassPath);
 	}
 
+	DisallowedObjectClassesToLoad.Reset();
+	AssetRegistry.GetDerivedClassNames(DisallowedObjectBaseClassPaths, TSet<FTopLevelAssetPath>(),
+		DisallowedObjectClassesToLoad);
+
 	TArray<FString> DisallowedAssetClassNamesToGather = GetConfigArray(TEXT("DisallowedAssetClassesToGather"));;
 	for (const FString& ClassName : DisallowedAssetClassNamesToGather)
 	{
@@ -139,6 +203,7 @@ TArray<FString> FIniCookedEditorPackageManager::GetConfigArray(const TCHAR* Key,
 
 void FIniCookedEditorPackageManager::GatherAllPackages(TArray<FName>& PackageNames) const
 {
+	check(bClassesInitialized); // Should be set by InitializeClasses from InitializeForCook
 	GatherAllPackagesExceptDisabled(PackageNames, DisabledPlugins);
 }
 
@@ -161,6 +226,11 @@ void FIniCookedEditorPackageManager::FilterGatheredPackages(TArray<FName>& Packa
 	PackageNames.Remove(NAME_None);
 }
 
+void FIniCookedEditorPackageManager::InitializeForCook()
+{
+	InitializeClasses();
+}
+
 void FIniCookedEditorPackageManager::GetEnginePackagesToCook(TArray<FName>& PackagesToCook) const
 {
 	for (const FString& Path : EngineAssetPaths)
@@ -192,9 +262,16 @@ void FIniCookedEditorPackageManager::GetProjectPackagesToCook(TArray<FName>& Pac
 
 bool FIniCookedEditorPackageManager::AllowObjectToBeCooked(const class UObject* Obj) const
 {
-	for (UClass* Class : DisallowedObjectClassesToLoad)
+	check(bClassesInitialized); // Should be set by InitializeClasses from InitializeForCook
+	const UClass* ObjAsClass = Cast<UClass>(Obj);
+
+	// A pointer to a disallowed native class is not filtered out, only instances of the native class are
+	// filtered out. For non-native classes, both the pointer to the non-native class and instances of the class
+	// are filtered out.
+	if (!ObjAsClass || !ObjAsClass->IsNative())
 	{
-		if (Obj->IsA(Class))
+		const UClass* ClassToCheck = ObjAsClass ? ObjAsClass : Obj->GetClass();
+		if (DisallowedObjectClassesToLoad.Contains(ClassToCheck->GetClassPathName()))
 		{
 			return false;
 		}
@@ -204,6 +281,7 @@ bool FIniCookedEditorPackageManager::AllowObjectToBeCooked(const class UObject*
 
 bool FIniCookedEditorPackageManager::AllowAssetToBeGathered(const struct FAssetData& AssetData) const
 {
+	check(bClassesInitialized); // Should be set by InitializeClasses from InitializeForCook
 	for (UClass* Class : DisallowedAssetClassesToGather)
 	{
 		if (AssetData.IsInstanceOf(Class))
diff --git a/Engine/Source/Developer/CookedEditor/Public/CookedEditorPackageManager.h b/Engine/Source/Developer/CookedEditor/Public/CookedEditorPackageManager.h
--- a/Engine/Source/Developer/CookedEditor/Public/CookedEditorPackageManager.h
+++ b/Engine/Source/Developer/CookedEditor/Public/CookedEditorPackageManager.h
@@ -76,6 +76,12 @@ public:
 	{
 	}
 
+	/**
+	 * Hook called when the manager's TargetPlatform is selected for cooking. Load any required assets or data.
+	 * May be called multiple times per process; it is called each time a cook starts for the platform.
+	 */
+	virtual void InitializeForCook() = 0;
+
 	/**
 	 * Override to control engine packages to cook
 	 */
@@ -154,15 +160,21 @@ class COOKEDEDITOR_API FIniCookedEditorPackageManager : public ICookedEditorPack
 	TArray<FString> EngineAssetPaths;
 	TArray<FString> ProjectAssetPaths;
 	TArray<FString> DisabledPlugins;
-	TArray<UClass*> DisallowedObjectClassesToLoad;
+	TSet<FTopLevelAssetPath> DisallowedObjectClassesToLoad;
 	TArray<UClass*> DisallowedAssetClassesToGather;
 	TArray<FString> DisallowedPathsToGather;
 
 	// true if this is a cooked cooker (false for cooker editor)
 	bool bIsCookedCooker;
 
+	// Guard to prevent multiple InitializeClasses calls
+	bool bClassesInitialized = false;
+
 	// gets an array from two sections, depending on bIsCookedCooker setting
 	TArray<FString> GetConfigArray(const TCHAR* Key) const;
+
+	// Load the classes specified from config settings that we need to reference
+	void InitializeClasses();
 public:
 
 	FIniCookedEditorPackageManager(bool bIsCookedCooker);
@@ -171,6 +183,7 @@ public:
 	static TArray<FString> GetConfigArray(const TCHAR* Key, bool bIsCookedCooker);
 
 	virtual void FilterGatheredPackages(TArray<FName>& PackageNames) const override;
+	virtual void InitializeForCook() override;
 	virtual void GetEnginePackagesToCook(TArray<FName>& PackagesToCook) const override;
 	virtual void GetProjectPackagesToCook(TArray<FName>& PackagesToCook) const override;
 	virtual bool AllowObjectToBeCooked(const class UObject* Obj) const override;
diff --git a/Engine/Source/Developer/CookedEditor/Public/CookedEditorTargetPlatformControls.h b/Engine/Source/Developer/CookedEditor/Public/CookedEditorTargetPlatformControls.h
--- a/Engine/Source/Developer/CookedEditor/Public/CookedEditorTargetPlatformControls.h
+++ b/Engine/Source/Developer/CookedEditor/Public/CookedEditorTargetPlatformControls.h
@@ -67,6 +67,11 @@ public:
 		return PackageManager->AllowObjectToBeCooked(Obj);
 	}
 
+	virtual void InitializeForCook() override
+	{
+		PackageManager->InitializeForCook();
+	}
+
 	virtual void GetExtraPackagesToCook(TArray<FName>& PackageNames) const override
 	{
 		if (AllowsEditorObjects())
@@ -158,6 +163,11 @@ public:
 		return PackageManager->AllowObjectToBeCooked(Obj);
 	}
 	
+	virtual void InitializeForCook() override
+	{
+		PackageManager->InitializeForCook();
+	}
+
 	virtual void GetExtraPackagesToCook(TArray<FName>& PackageNames) const override
 	{
 		if (AllowsEditorObjects())
diff --git a/Engine/Source/Developer/TargetPlatform/Public/Common/TargetPlatformBase.h b/Engine/Source/Developer/TargetPlatform/Public/Common/TargetPlatformBase.h
--- a/Engine/Source/Developer/TargetPlatform/Public/Common/TargetPlatformBase.h
+++ b/Engine/Source/Developer/TargetPlatform/Public/Common/TargetPlatformBase.h
@@ -278,6 +278,11 @@ public:
 		return false; 
 	}
 
+	virtual void InitializeForCook() override
+	{
+
+	}
+
 	virtual void GetExtraPackagesToCook(TArray<FName>& PackageNames) const override
 	{
 
@@ -982,6 +987,10 @@ public:
 	{
 		return TargetPlatformControls->CopyFileToTarget(DeviceId, HostFilename, TargetFilename, CustomPlatformData);
 	}
+	virtual void InitializeForCook() override
+	{
+		TargetPlatformControls->InitializeForCook();
+	}
 	virtual void GetExtraPackagesToCook(TArray<FName>& PackageNames) const override
 	{
 		TargetPlatformControls->GetExtraPackagesToCook(PackageNames);
diff --git a/Engine/Source/Developer/TargetPlatform/Public/Common/TargetPlatformControlsBase.h b/Engine/Source/Developer/TargetPlatform/Public/Common/TargetPlatformControlsBase.h
--- a/Engine/Source/Developer/TargetPlatform/Public/Common/TargetPlatformControlsBase.h
+++ b/Engine/Source/Developer/TargetPlatform/Public/Common/TargetPlatformControlsBase.h
@@ -195,6 +195,11 @@ public:
 		return false;
 	}
 
+	virtual void InitializeForCook() override
+	{
+
+	}
+
 	virtual void GetExtraPackagesToCook(TArray<FName>& PackageNames) const override
 	{
 
diff --git a/Engine/Source/Developer/TargetPlatform/Public/Interfaces/ITargetPlatformControls.h b/Engine/Source/Developer/TargetPlatform/Public/Interfaces/ITargetPlatformControls.h
--- a/Engine/Source/Developer/TargetPlatform/Public/Interfaces/ITargetPlatformControls.h
+++ b/Engine/Source/Developer/TargetPlatform/Public/Interfaces/ITargetPlatformControls.h
@@ -478,6 +478,12 @@ public:
 	 */
 	virtual bool CopyFileToTarget(const FString& DeviceId, const FString& HostFilename, const FString& TargetFilename, const TMap<FString, FString>& CustomPlatformData) = 0;
 
+	/**
+	 * Hook called when a TargetPlatform is selected for cooking. Load any required assets or data.
+	 * May be called multiple times per process; it is called each time a cook starts for the platform.
+	 */
+	virtual void InitializeForCook() = 0;
+
 	/**
 	 * Gets a list of package names to cook when cooking this platform
 	 */
diff --git a/Engine/Source/Editor/UnrealEd/Private/CookOnTheFlyServer.cpp b/Engine/Source/Editor/UnrealEd/Private/CookOnTheFlyServer.cpp
--- a/Engine/Source/Editor/UnrealEd/Private/CookOnTheFlyServer.cpp
+++ b/Engine/Source/Editor/UnrealEd/Private/CookOnTheFlyServer.cpp
@@ -803,6 +803,7 @@ void UCookOnTheFlyServer::StartCookOnTheFlySessionFromGameThread(ITargetPlatform
 		InitializeSession();
 	}
 	PlatformManager->AddSessionPlatform(*this, TargetPlatform);
+	TargetPlatform->InitializeForCook();
 	ResetCook({ TPair<const ITargetPlatform*,bool>{TargetPlatform, true /* bResetResults */} });
 
 	// Blocking on the AssetRegistry needs to wait until the session starts because it needs all plugins loaded.
@@ -12344,6 +12345,7 @@ void UCookOnTheFlyServer::SelectSessionPlatforms(FBeginCookContext& BeginContext
 	for (FBeginCookContextPlatform& PlatformContext : BeginContext.PlatformContexts)
 	{
 		PlatformContext.PlatformData = PlatformManager->GetPlatformData(PlatformContext.TargetPlatform);
+		PlatformContext.TargetPlatform->InitializeForCook();
 	}
 }
 
