Added logging when an asset within a package is rejected when cooking a cooked editor.

[RN] Updated cook logs to include assets skipped due to their class type being in a disallow list.

diff --git a/Engine/Source/Developer/CookedEditor/Private/CookedEditorPackageManager.cpp b/Engine/Source/Developer/CookedEditor/Private/CookedEditorPackageManager.cpp
--- a/Engine/Source/Developer/CookedEditor/Private/CookedEditorPackageManager.cpp
+++ b/Engine/Source/Developer/CookedEditor/Private/CookedEditorPackageManager.cpp
@@ -4,19 +4,17 @@
 #include "Algo/Sort.h"
 #include "Algo/Unique.h"
 #include "AssetRegistry/AssetData.h"
-#include "CoreMinimal.h"
-#include "Logging/LogMacros.h"
 #include "AssetRegistry/AssetRegistryModule.h"
 #include "AssetRegistry/IAssetRegistry.h"
-#include "AssetRegistry/AssetRegistryModule.h"
+#include "CookerSettings.h"
+#include "CoreMinimal.h"
 #include "Engine/AssetManager.h"
 #include "GameDelegates.h"
+#include "Logging/LogMacros.h"
 
 DEFINE_LOG_CATEGORY_STATIC(LogCookedEditorTargetPlatform, Log, All)
 
 
-
-
 TUniquePtr<ICookedEditorPackageManager> ICookedEditorPackageManager::FactoryForTargetPlatform(bool bIsCookedCooker)
 {
 	if (FGameDelegates::Get().GetCookedEditorPackageManagerFactoryDelegate().IsBound())
@@ -107,6 +105,11 @@ void FIniCookedEditorPackageManager::InitializeClasses()
 	}
 	bClassesInitialized = true;
 
+	if (const IConsoleVariable* DisplayModeCVar = IConsoleManager::Get().FindConsoleVariable(TEXT("cook.displaymode")))
+	{
+		CookProgressDisplayMode = DisplayModeCVar->GetInt();
+	}
+
 	IAssetRegistry& AssetRegistry = IAssetRegistry::GetChecked();
 	// SearchAllAssets should be a noop because the cooker already did it, but run it just in case.
 	AssetRegistry.SearchAllAssets(true /* bSynchronousSearch */);
@@ -273,6 +276,10 @@ bool FIniCookedEditorPackageManager::AllowObjectToBeCooked(const class UObject*
 		const UClass* ClassToCheck = ObjAsClass ? ObjAsClass : Obj->GetClass();
 		if (DisallowedObjectClassesToLoad.Contains(ClassToCheck->GetClassPathName()))
 		{
+			if(CookProgressDisplayMode != (int32)ECookProgressDisplayMode::Nothing)
+			{
+				UE_LOG(LogCookedEditorTargetPlatform, Log, TEXT("Object: %s will not cook due to disallowed class: %s"), *Obj->GetFName().ToString(), *ClassToCheck->GetClassPathName().ToString());
+			}
 			return false;
 		}
 	}
@@ -286,6 +293,10 @@ bool FIniCookedEditorPackageManager::AllowAssetToBeGathered(const struct FAssetD
 	{
 		if (AssetData.IsInstanceOf(Class))
 		{
+			if (CookProgressDisplayMode != (int32)ECookProgressDisplayMode::Nothing)
+			{
+				UE_LOG(LogCookedEditorTargetPlatform, Log, TEXT("Asset: %s will not be gathered due to disallowed class: %s"), *AssetData.PackageName.ToString(), *Class->GetClassPathName().ToString());
+			}
 			return false;
 		}
 	}
@@ -302,6 +313,3 @@ bool FIniCookedEditorPackageManager::AllowProjectPluginContentToBeCooked(const T
 {
 	return true;
 }
-
-
-
diff --git a/Engine/Source/Developer/CookedEditor/Public/CookedEditorPackageManager.h b/Engine/Source/Developer/CookedEditor/Public/CookedEditorPackageManager.h
--- a/Engine/Source/Developer/CookedEditor/Public/CookedEditorPackageManager.h
+++ b/Engine/Source/Developer/CookedEditor/Public/CookedEditorPackageManager.h
@@ -170,6 +170,9 @@ class COOKEDEDITOR_API FIniCookedEditorPackageManager : public ICookedEditorPack
 	// Guard to prevent multiple InitializeClasses calls
 	bool bClassesInitialized = false;
 
+	// Used to control logging of certain warnings.
+	int32 CookProgressDisplayMode = 0;
+
 	// gets an array from two sections, depending on bIsCookedCooker setting
 	TArray<FString> GetConfigArray(const TCHAR* Key) const;
 
