Don't include windows things if not windows

#rnx

diff --git a/Engine/Plugins/Experimental/AVCodecs/NVCodecs/Source/NVDEC/NVDEC.Build.cs b/Engine/Plugins/Experimental/AVCodecs/NVCodecs/Source/NVDEC/NVDEC.Build.cs
--- a/Engine/Plugins/Experimental/AVCodecs/NVCodecs/Source/NVDEC/NVDEC.Build.cs
+++ b/Engine/Plugins/Experimental/AVCodecs/NVCodecs/Source/NVDEC/NVDEC.Build.cs
@@ -28,6 +28,9 @@ public class NVDEC: ModuleRules
 		});
 
 		AddEngineThirdPartyPrivateStaticDependencies(Target, "Vulkan");
-		AddEngineThirdPartyPrivateStaticDependencies(Target, "DX11");
+		if (Target.IsInPlatformGroup(UnrealPlatformGroup.Windows))
+		{
+			AddEngineThirdPartyPrivateStaticDependencies(Target, "DX11");
+		}
 	}
 }
diff --git a/Engine/Plugins/Experimental/AVCodecs/NVCodecs/Source/NVENC/NVENC.Build.cs b/Engine/Plugins/Experimental/AVCodecs/NVCodecs/Source/NVENC/NVENC.Build.cs
--- a/Engine/Plugins/Experimental/AVCodecs/NVCodecs/Source/NVENC/NVENC.Build.cs
+++ b/Engine/Plugins/Experimental/AVCodecs/NVCodecs/Source/NVENC/NVENC.Build.cs
@@ -28,6 +28,9 @@ public class NVENC: ModuleRules
 		});
 
 		AddEngineThirdPartyPrivateStaticDependencies(Target, "Vulkan");
-		AddEngineThirdPartyPrivateStaticDependencies(Target, "DX11");
+		if (Target.IsInPlatformGroup(UnrealPlatformGroup.Windows))
+		{
+			AddEngineThirdPartyPrivateStaticDependencies(Target, "DX11");
+		}
 	}
 }
diff --git a/Engine/Plugins/Media/HardwareEncoders/Source/EncoderNVENC/EncoderNVENC.Build.cs b/Engine/Plugins/Media/HardwareEncoders/Source/EncoderNVENC/EncoderNVENC.Build.cs
--- a/Engine/Plugins/Media/HardwareEncoders/Source/EncoderNVENC/EncoderNVENC.Build.cs
+++ b/Engine/Plugins/Media/HardwareEncoders/Source/EncoderNVENC/EncoderNVENC.Build.cs
@@ -31,6 +31,9 @@ public class EncoderNVENC: ModuleRules
 		PrivateIncludePaths.Add(Path.Combine(EngineSourceDirectory, "Source/Runtime/AVEncoder/Private"));
 
 		AddEngineThirdPartyPrivateStaticDependencies(Target, "Vulkan");
-		AddEngineThirdPartyPrivateStaticDependencies(Target, "DX11");
+		if (Target.IsInPlatformGroup(UnrealPlatformGroup.Windows))
+		{
+			AddEngineThirdPartyPrivateStaticDependencies(Target, "DX11");
+		}
 	}
 }
diff --git a/Engine/Plugins/VirtualProduction/TextureShare/Source/TextureShare/TextureShare.Build.cs b/Engine/Plugins/VirtualProduction/TextureShare/Source/TextureShare/TextureShare.Build.cs
--- a/Engine/Plugins/VirtualProduction/TextureShare/Source/TextureShare/TextureShare.Build.cs
+++ b/Engine/Plugins/VirtualProduction/TextureShare/Source/TextureShare/TextureShare.Build.cs
@@ -78,8 +78,11 @@ public class TextureShare : ModuleRules
 			PublicDefinitions.Add("TEXTURESHARE_VULKAN=0");
 		}
 
-		AddEngineThirdPartyPrivateStaticDependencies(Target, "DX11");
-		AddEngineThirdPartyPrivateStaticDependencies(Target, "DX12");
+		if (Target.IsInPlatformGroup(UnrealPlatformGroup.Windows))
+		{
+			AddEngineThirdPartyPrivateStaticDependencies(Target, "DX11");
+			AddEngineThirdPartyPrivateStaticDependencies(Target, "DX12");
+		}
 
 		if (Target.bBuildEditor == true)
 		{
diff --git a/Engine/Source/Developer/CookedEditor/CookedEditor.Build.cs b/Engine/Source/Developer/CookedEditor/CookedEditor.Build.cs
--- a/Engine/Source/Developer/CookedEditor/CookedEditor.Build.cs
+++ b/Engine/Source/Developer/CookedEditor/CookedEditor.Build.cs
@@ -27,8 +27,16 @@ public class CookedEditor : ModuleRules
 				"TargetPlatform",
 			});
 
-		PublicIncludePathModuleNames.Add("WindowsTargetPlatformSettings");
-		PublicIncludePathModuleNames.Add("WindowsTargetPlatformControls");
+		if (IsPlatformAvailable(UnrealTargetPlatform.Win64))
+		{
+			PublicDefinitions.Add("COOKEDEDITOR_WITH_WINDOWSTARGETPLATFORM=1");
+			PublicIncludePathModuleNames.Add("WindowsTargetPlatformSettings");
+			PublicIncludePathModuleNames.Add("WindowsTargetPlatformControls");
+		}
+		else
+		{
+			PublicDefinitions.Add("COOKEDEDITOR_WITH_WINDOWSTARGETPLATFORM=0");
+		}
 
 		if (IsPlatformAvailable(UnrealTargetPlatform.Linux, bIgnoreSDKCheck: true))
 		{
diff --git a/Engine/Source/Developer/CookedEditor/Public/CookedEditorPackageManager.h b/Engine/Source/Developer/CookedEditor/Public/CookedEditorPackageManager.h
--- a/Engine/Source/Developer/CookedEditor/Public/CookedEditorPackageManager.h
+++ b/Engine/Source/Developer/CookedEditor/Public/CookedEditorPackageManager.h
@@ -6,8 +6,10 @@
 #include "Common/TargetPlatformBase.h"
 #include "Interfaces/IPluginManager.h"
 
-#include "GenericWindowsTargetPlatformControls.h"
-#include "Windows/WindowsPlatformProperties.h"
+#if COOKEDEDITOR_WITH_WINDOWSTARGETPLATFORM
+	#include "GenericWindowsTargetPlatformControls.h"
+	#include "Windows/WindowsPlatformProperties.h"
+#endif
 
 #if COOKEDEDITOR_WITH_LINUXTARGETPLATFORM
 	#include "LinuxTargetPlatformSettings.h"
@@ -19,8 +21,10 @@
 	#include "GenericMacTargetPlatformControls.h"
 #endif
 
-typedef TGenericWindowsTargetPlatformSettings<FWindowsPlatformProperties<false, false, false>> FWindowsEditorTargetPlatformSettingsParent;
-typedef TGenericWindowsTargetPlatformControls<FWindowsPlatformProperties<false, false, false>> FWindowsEditorTargetPlatformControlsParent;
+#if COOKEDEDITOR_WITH_WINDOWSTARGETPLATFORM
+	typedef TGenericWindowsTargetPlatformSettings<FWindowsPlatformProperties<false, false, false>> FWindowsEditorTargetPlatformSettingsParent;
+	typedef TGenericWindowsTargetPlatformControls<FWindowsPlatformProperties<false, false, false>> FWindowsEditorTargetPlatformControlsParent;
+#endif
 
 #if COOKEDEDITOR_WITH_LINUXTARGETPLATFORM
 typedef TLinuxTargetPlatformSettings<FLinuxPlatformProperties<false, false, false, false>> FLinuxEditorTargetPlatformSettingsParent;
